Updating translation for ja #jira none

diff --git a/Templates/TP_AEC_ProdConfigBP/Config/TemplateDefs.ini b/Templates/TP_AEC_ProdConfigBP/Config/TemplateDefs.ini
--- a/Templates/TP_AEC_ProdConfigBP/Config/TemplateDefs.ini
+++ b/Templates/TP_AEC_ProdConfigBP/Config/TemplateDefs.ini
@@ -11,12 +11,13 @@ Categories=AEC
 LocalizedDisplayNames=(Language="en", Text="Design Configurator")
 LocalizedDisplayNames=(Language="ko", Text="제품 컨피규레이터")
 LocalizedDisplayNames=(Language="zh-Hans", Text="设计配置器")
+LocalizedDisplayNames=(Language="ja", Text="デザインコンフィギュレーター")
 
 ; The localized descriptions for this template, which is displayed when the template is selected.
 LocalizedDescriptions=(Language="en", Text="This project template uses the Variant Manager, UMG and Blueprint to build a project allowing you to toggle between different object states such as visibility, trigger animation sequences, switch between views or toggle between different design options.\n\nThe user interface is dynamically created from the 3d model and the variants present in the level. This allows for users to swap the provided model for user-made assets and have the configurator's UI adapting to the new data.")
 LocalizedDescriptions=(Language="ko", Text="이 프로젝트 템플릿은 일반적인 제품 컨피규레이터 제작을 위해 베리언트 매니저, UMG 그리고 블루프린트를 활용합니다.\n\n유저 인터페이스는 3D 모델과 레벨 내 베리언트를 기반으로 다이내믹하게 생성됩니다. 이를 통해 사용자들은 제공된 모델을 직접 제작한 애셋과 교체하는 것만으로 새로운 데이터에 맞는 컨피규레이터 UI를 사용할 수 있습니다.")
 LocalizedDescriptions=(Language="zh-Hans", Text="该项目模板使用变体管理器，UMG和蓝图构建了一个可在不同目标间切换的项目，如可见性、触发动画序列、切换视图模式或是开启/关闭不同的设计选项。用户界面是用关卡中的3D模型和变体动态创建的。因此用户能够用自制资产替换现有模型，从而使配置器的UI适配新数据。")
-LocalizedDescriptions=(Language="ja", Text="このプロジェクト テンプレートでは、Variant Manager、UMG、ブループリントを利用して、汎用的な製品コンフィギュレーターを構築できます。\n\nユーザーインターフェイスは、3D モデルとレベル内に表示されているバリアントから動的に作成されるため、ユーザーは、あらかじめ備わっているモデルをユーザー作成によるアセットと交換し、コンフィギュレーターの UI をその新たなデータに適合させることができます。")
+LocalizedDescriptions=(Language="ja", Text="このプロジェクト テンプレートは Variant Manager、UMG、ブループリントを使用して、様々なオブジェクトの可視性の切替、アニメーション シーケンスの呼び出し、視点の変更、デザイン設定の切替ができるプロジェクトを構築するものです。ユーザーインターフェイスは、3D モデルとレベル内に表示されているバリアントから動的に作成されます。これにより、あらかじめ備わっているモデルをユーザー作成によるアセットと交換すれば、コンフィギュレーターの UI はその新たなデータに適合します。")
 
 
 ; The sort key to use for the template to bypass the default alphabetical sorting.
diff --git a/Templates/TP_AEC_ProdConfigBP/manifest.json b/Templates/TP_AEC_ProdConfigBP/manifest.json
--- a/Templates/TP_AEC_ProdConfigBP/manifest.json
+++ b/Templates/TP_AEC_ProdConfigBP/manifest.json
@@ -14,6 +14,10 @@
 			"Language": "zh-Hans",
 			"Text": "设计配置器"
 		}
+		{
+			"Language": "ja",
+			"Text": "デザインコンフィギュレーター"
+		}
 	],
 	"Description":
 	[
@@ -31,7 +35,7 @@
 		},
 		{
 			"Language": "ja",
-			"Text": "このプロジェクト テンプレートでは、Variant Manager、UMG、ブループリントを利用して、汎用的な製品コンフィギュレーターを構築できます。\n\nユーザーインターフェイスは、3D モデルとレベル内に表示されているバリアントから動的に作成されるため、ユーザーは、あらかじめ備わっているモデルをユーザー作成によるアセットと交換し、コンフィギュレーターの UI をその新たなデータに適合させることができます。"
+			"Text": "このプロジェクト テンプレートは Variant Manager、UMG、ブループリントを使用して、様々なオブジェクトの可視性の切替、アニメーション シーケンスの呼び出し、視点の変更、デザイン設定の切替ができるプロジェクトを構築するものです。ユーザーインターフェイスは、3D モデルとレベル内に表示されているバリアントから動的に作成されます。これにより、あらかじめ備わっているモデルをユーザー作成によるアセットと交換すれば、コンフィギュレーターの UI はその新たなデータに適合します。"
 		}
 	],
 	"AssetTypes": [],
