Updating translation for ko #jira none

diff --git a/Templates/TP_AEC_ProdConfigBP/Config/TemplateDefs.ini b/Templates/TP_AEC_ProdConfigBP/Config/TemplateDefs.ini
--- a/Templates/TP_AEC_ProdConfigBP/Config/TemplateDefs.ini
+++ b/Templates/TP_AEC_ProdConfigBP/Config/TemplateDefs.ini
@@ -9,13 +9,13 @@ Categories=AEC
 
 ; The localized display names for this template. Each entry specifies a Language code and Unicode text to display in the New Project dialog.
 LocalizedDisplayNames=(Language="en", Text="Design Configurator")
-LocalizedDisplayNames=(Language="ko", Text="제품 컨피규레이터")
+LocalizedDisplayNames=(Language="ko", Text="디자인 컨피규레이터")
 LocalizedDisplayNames=(Language="zh-Hans", Text="设计配置器")
 LocalizedDisplayNames=(Language="ja", Text="デザインコンフィギュレーター")
 
 ; The localized descriptions for this template, which is displayed when the template is selected.
 LocalizedDescriptions=(Language="en", Text="This project template uses the Variant Manager, UMG and Blueprint to build a project allowing you to toggle between different object states such as visibility, trigger animation sequences, switch between views or toggle between different design options.\n\nThe user interface is dynamically created from the 3d model and the variants present in the level. This allows for users to swap the provided model for user-made assets and have the configurator's UI adapting to the new data.")
-LocalizedDescriptions=(Language="ko", Text="이 프로젝트 템플릿은 일반적인 제품 컨피규레이터 제작을 위해 베리언트 매니저, UMG 그리고 블루프린트를 활용합니다.\n\n유저 인터페이스는 3D 모델과 레벨 내 베리언트를 기반으로 다이내믹하게 생성됩니다. 이를 통해 사용자들은 제공된 모델을 직접 제작한 애셋과 교체하는 것만으로 새로운 데이터에 맞는 컨피규레이터 UI를 사용할 수 있습니다.")
+LocalizedDescriptions=(Language="ko", Text="이 프로젝트 템플릿은 베리언트 매니저, UMG와 블루프린트를 사용한 프로젝트로 가시성,  애니메이션 시퀀스 트리거, 뷰 모드 또는 다양한 디자인 옵션 간의 전환 등과 같이 다른 오브젝트 스테이트를 전환할 수 있습니다. 유저 인터페이스는 레벨 내 3D 모델과 베리언트를 기반으로 동적으로 생성됩니다. 이를 통해 사용자들은 기본 제공된 모델을 직접 제작한 에셋과 교체하고 새로운 데이터에 맞게 컨피규레이터 UI를 업데이트할 수 있습니다.")
 LocalizedDescriptions=(Language="zh-Hans", Text="该项目模板使用变体管理器，UMG和蓝图构建了一个可在不同目标间切换的项目，如可见性、触发动画序列、切换视图模式或是开启/关闭不同的设计选项。用户界面是用关卡中的3D模型和变体动态创建的。因此用户能够用自制资产替换现有模型，从而使配置器的UI适配新数据。")
 LocalizedDescriptions=(Language="ja", Text="このプロジェクト テンプレートは Variant Manager、UMG、ブループリントを使用して、様々なオブジェクトの可視性の切替、アニメーション シーケンスの呼び出し、視点の変更、デザイン設定の切替ができるプロジェクトを構築するものです。ユーザーインターフェイスは、3D モデルとレベル内に表示されているバリアントから動的に作成されます。これにより、あらかじめ備わっているモデルをユーザー作成によるアセットと交換すれば、コンフィギュレーターの UI はその新たなデータに適合します。")
 
diff --git a/Templates/TP_AEC_ProdConfigBP/manifest.json b/Templates/TP_AEC_ProdConfigBP/manifest.json
--- a/Templates/TP_AEC_ProdConfigBP/manifest.json
+++ b/Templates/TP_AEC_ProdConfigBP/manifest.json
@@ -8,7 +8,7 @@
 		},
 		{
 			"Language": "ko",
-			"Text": "제품 컨피규레이터"
+			"Text": "디자인 컨피규레이터"
 		},
 		{
 			"Language": "zh-Hans",
@@ -27,7 +27,7 @@
 		},
 		{
 			"Language": "ko",
-			"Text": "이 프로젝트 템플릿은 일반적인 제품 컨피규레이터 제작을 위해 베리언트 매니저, UMG 그리고 블루프린트를 활용합니다.\n\n유저 인터페이스는 3D 모델과 레벨 내 베리언트를 기반으로 다이내믹하게 생성됩니다. 이를 통해 사용자들은 제공된 모델을 직접 제작한 애셋과 교체하는 것만으로 새로운 데이터에 맞는 컨피규레이터 UI를 사용할 수 있습니다."
+			"Text": "이 프로젝트 템플릿은 베리언트 매니저, UMG와 블루프린트를 사용한 프로젝트로 가시성,  애니메이션 시퀀스 트리거, 뷰 모드 또는 다양한 디자인 옵션 간의 전환 등과 같이 다른 오브젝트 스테이트를 전환할 수 있습니다. 유저 인터페이스는 레벨 내 3D 모델과 베리언트를 기반으로 동적으로 생성됩니다. 이를 통해 사용자들은 기본 제공된 모델을 직접 제작한 에셋과 교체하고 새로운 데이터에 맞게 컨피규레이터 UI를 업데이트할 수 있습니다."
 		},
 		{
 			"Language": "zh-Hans",
